Interchange: Dedicated plugin for the OpenVDB translator and its dependencies.

diff --git a/Engine/Plugins/Interchange/Extensions/OpenVDB/InterchangeOpenVDB.uplugin b/Engine/Plugins/Interchange/Extensions/OpenVDB/InterchangeOpenVDB.uplugin
new file mode 100644
--- /dev/null
+++ b/Engine/Plugins/Interchange/Extensions/OpenVDB/InterchangeOpenVDB.uplugin
@@ -0,0 +1,35 @@
+{
+	"FileVersion": 3,
+	"Version": 1,
+	"VersionName": "1.0",
+	"FriendlyName": "Interchange OpenVDB Support",
+	"Description": "Allows translation of OpenVDB files via the Interchange framework",
+	"Category": "Importers",
+	"CreatedBy": "Epic Games, Inc.",
+	"CreatedByURL": "https://epicgames.com",
+	"DocsURL": "",
+	"MarketplaceURL": "",
+	"SupportURL": "",
+	"EnabledByDefault": false,
+	"CanContainContent": false,
+	"IsBetaVersion": false,
+	"Installed": false,
+	"IsExperimentalVersion": true,
+	"SupportedTargetPlatforms": [
+		"Win64",
+		"Linux"
+	],
+	"Modules": [
+		{
+			"Name": "InterchangeOpenVDBImport",
+			"Type": "Editor",
+			"LoadingPhase": "Default"
+		}
+	],
+	"Plugins": [
+		{
+			"Name": "Interchange",
+			"Enabled": true
+		}
+	]
+}
\ No newline at end of file
diff --git a/Engine/Plugins/Interchange/Extensions/OpenVDB/Source/Import/InterchangeOpenVDBImport.Build.cs b/Engine/Plugins/Interchange/Extensions/OpenVDB/Source/Import/InterchangeOpenVDBImport.Build.cs
new file mode 100644
--- /dev/null
+++ b/Engine/Plugins/Interchange/Extensions/OpenVDB/Source/Import/InterchangeOpenVDBImport.Build.cs
@@ -0,0 +1,32 @@
+// Copyright Epic Games, Inc. All Rights Reserved.
+
+using System.IO;
+
+namespace UnrealBuildTool.Rules
+{
+	public class InterchangeOpenVDBImport : ModuleRules
+	{
+		public InterchangeOpenVDBImport(ReadOnlyTargetRules Target) : base(Target)
+		{
+			PublicDependencyModuleNames.AddRange(
+				new string[]
+				{
+					"Core",
+					"CoreUObject",
+					"Engine",
+				}
+			);
+
+			PrivateDependencyModuleNames.AddRange(
+				new string[]
+				{
+					"InterchangeCore", 		// For UInterchangeBaseNode, base translator
+					"InterchangeImport", 	// For the translator settings, pipeline interface, etc.
+					"InterchangeNodes", 	// For the actual volume nodes
+					"InterchangeEngine", 	// For UInterchangeManager::GetInterchangeManager
+					"SparseVolumeTexture",	// For the legacy OpenVDB importer reference, that actually uses the OpenVDB SDK
+				}
+			);
+		}
+	}
+}
diff --git a/Engine/Plugins/Interchange/Extensions/OpenVDB/Source/Import/Private/InterchangeOpenVDBImportLog.h b/Engine/Plugins/Interchange/Extensions/OpenVDB/Source/Import/Private/InterchangeOpenVDBImportLog.h
new file mode 100644
--- /dev/null
+++ b/Engine/Plugins/Interchange/Extensions/OpenVDB/Source/Import/Private/InterchangeOpenVDBImportLog.h
@@ -0,0 +1,7 @@
+// Copyright Epic Games, Inc. All Rights Reserved.
+
+#pragma once
+
+#include "CoreMinimal.h"
+
+DECLARE_LOG_CATEGORY_EXTERN(LogInterchangeOpenVDBImport, Log, All);
diff --git a/Engine/Plugins/Interchange/Extensions/OpenVDB/Source/Import/Private/InterchangeOpenVDBImportModule.cpp b/Engine/Plugins/Interchange/Extensions/OpenVDB/Source/Import/Private/InterchangeOpenVDBImportModule.cpp
new file mode 100644
--- /dev/null
+++ b/Engine/Plugins/Interchange/Extensions/OpenVDB/Source/Import/Private/InterchangeOpenVDBImportModule.cpp
@@ -0,0 +1,45 @@
+// Copyright Epic Games, Inc. All Rights Reserved.
+
+#include "InterchangeOpenVDBImportModule.h"
+
+#include "InterchangeOpenVDBImportLog.h"
+#include "InterchangeOpenVDBTranslator.h"
+
+#include "Engine/Engine.h"
+#include "InterchangeManager.h"
+#include "Misc/CoreDelegates.h"
+#include "Modules/ModuleManager.h"
+
+DEFINE_LOG_CATEGORY(LogInterchangeOpenVDBImport);
+
+void FInterchangeOpenVDBImportModule::StartupModule()
+{
+	using namespace UE::Interchange;
+
+	auto RegisterItems = [this]()
+	{
+		UInterchangeManager& InterchangeManager = UInterchangeManager::GetInterchangeManager();
+
+		// Volumes
+#if WITH_EDITOR && OPENVDB_AVAILABLE
+		// Editor and OpenVDB-only because we use GetOpenVDBGridInfo and ConvertOpenVDBToSparseVolumeTexture
+		// which use OpenVDB and are in an Editor-only module
+		InterchangeManager.RegisterTranslator(UInterchangeOpenVDBTranslator::StaticClass());
+#endif
+	};
+
+	if (GEngine)
+	{
+		RegisterItems();
+	}
+	else
+	{
+		FCoreDelegates::OnPostEngineInit.AddLambda(RegisterItems);
+	}
+}
+
+void FInterchangeOpenVDBImportModule::ShutdownModule()
+{
+}
+
+IMPLEMENT_MODULE(FInterchangeOpenVDBImportModule, InterchangeOpenVDBImport)
diff --git a/Engine/Plugins/Interchange/Extensions/OpenVDB/Source/Import/Private/InterchangeOpenVDBTranslator.cpp b/Engine/Plugins/Interchange/Extensions/OpenVDB/Source/Import/Private/InterchangeOpenVDBTranslator.cpp
new file mode 100644
--- /dev/null
+++ b/Engine/Plugins/Interchange/Extensions/OpenVDB/Source/Import/Private/InterchangeOpenVDBTranslator.cpp
@@ -0,0 +1,543 @@
+// Copyright Epic Games, Inc. All Rights Reserved.
+
+#include "InterchangeOpenVDBTranslator.h"
+
+#include "InterchangeManager.h"
+#include "InterchangeOpenVDBImportLog.h"
+#include "InterchangeSceneNode.h"
+#include "InterchangeVolumeNode.h"
+#include "Usd/InterchangeUsdDefinitions.h"
+#include "Volume/InterchangeVolumeTranslatorSettings.h"
+
+#include "Async/ParallelFor.h"
+#include "HAL/IConsoleManager.h"
+
+#if WITH_EDITOR
+#include "OpenVDBImportOptions.h"
+#include "SparseVolumeTextureOpenVDBUtility.h"
+#endif	  // WITH_EDITOR
+
+#include UE_INLINE_GENERATED_CPP_BY_NAME(InterchangeOpenVDBTranslator)
+
+#define LOCTEXT_NAMESPACE "InterchangeOpenVDBTranslator"
+
+static bool GInterchangeEnableOpenVDBImport = false;
+static FAutoConsoleVariableRef CVarInterchangeEnableOpenVDBImport(
+	TEXT("Interchange.FeatureFlags.Import.OpenVDB"),
+	GInterchangeEnableOpenVDBImport,
+	TEXT("Whether OpenVDB support is enabled via Interchange")
+);
+
+static int32 GMaxParallelFileReads = 16;
+static FAutoConsoleVariableRef CVarNumPerPrimLocks(
+	TEXT("Interchange.FeatureFlags.Import.OpenVDB.NumFileReadLocks"),
+	GMaxParallelFileReads,
+	TEXT("Maximum number of .vdb files that can be read concurrently. This uses system mutexes, which are limited and vary depending on your system."
+	),
+	ECVF_ReadOnly	 // We'll keep a static array of these for simplicity, so changing it during the session would be annoying to handle
+);
+
+namespace UE::InterchangeOpenVDBTranslator::Private
+{
+#if WITH_EDITOR
+	struct FOpenVDBFileInfo
+	{
+		TArray64<uint8> FileBytes;
+		TArray<FOpenVDBGridInfo> GridInfo;
+
+	private:
+		friend class UInterchangeOpenVDBTranslatorImpl;
+
+		// Index of our lock within the Impl's FileLocks. Having multiple locks allows concurrent file reads
+		int32 LockIndex = INDEX_NONE;
+
+		// We set this to true if we failed to load/open this file, so that we don't try again later and violate the
+		// assumption that the FileInfos are going to be read-only after they're returned from GetOrLoadFileInfo
+		bool bFailed = false;
+	};
+
+	EVolumeGridElementType GridTypeToInterchangeGridType(EOpenVDBGridType Type)
+	{
+		switch (Type)
+		{
+			case EOpenVDBGridType::Half:
+			case EOpenVDBGridType::Half2:
+			case EOpenVDBGridType::Half3:
+			case EOpenVDBGridType::Half4:
+			{
+				return EVolumeGridElementType::Half;
+			}
+			case EOpenVDBGridType::Float:
+			case EOpenVDBGridType::Float2:
+			case EOpenVDBGridType::Float3:
+			case EOpenVDBGridType::Float4:
+			{
+				return EVolumeGridElementType::Float;
+			}
+			case EOpenVDBGridType::Double:
+			case EOpenVDBGridType::Double2:
+			case EOpenVDBGridType::Double3:
+			case EOpenVDBGridType::Double4:
+			{
+				return EVolumeGridElementType::Double;
+			}
+			default:
+			{
+				return EVolumeGridElementType::Unknown;
+			}
+		}
+	}
+
+	class UInterchangeOpenVDBTranslatorImpl
+	{
+	public:
+		UInterchangeOpenVDBTranslatorImpl()
+		{
+			FileLocks = new FRWLock[GMaxParallelFileReads]();
+		}
+
+		~UInterchangeOpenVDBTranslatorImpl()
+		{
+			delete[] FileLocks;
+		}
+
+		UInterchangeOpenVDBTranslatorImpl(const UInterchangeOpenVDBTranslatorImpl& Other) = delete;
+		UInterchangeOpenVDBTranslatorImpl& operator=(const UInterchangeOpenVDBTranslatorImpl& Other) = delete;
+
+	public:
+		void ClearFileInfoEntries()
+		{
+			FWriteScopeLock WriteLock(FileNameToInfoLock);
+
+			FileNameToInfo.Reset();
+		}
+
+		void EnsureFileInfoEntries(const TArray<FString>& Filenames)
+		{
+			FWriteScopeLock WriteLock(FileNameToInfoLock);
+
+			for (const FString& Filename : Filenames)
+			{
+				if (FileNameToInfo.Contains(Filename))
+				{
+					continue;
+				}
+
+				TUniquePtr<FOpenVDBFileInfo> NewInfo = MakeUnique<FOpenVDBFileInfo>();
+				NewInfo->LockIndex = NextFileLockIndex++;
+
+				FileNameToInfo.Add(Filename, MoveTemp(NewInfo));
+			}
+		}
+
+		const FOpenVDBFileInfo* GetOrLoadFileInfo(const FString& Filename)
+		{
+			FOpenVDBFileInfo* Info = nullptr;
+			{
+				FReadScopeLock ReadLock(FileNameToInfoLock);
+
+				if (TUniquePtr<FOpenVDBFileInfo>* FoundInfo = FileNameToInfo.Find(Filename))
+				{
+					Info = FoundInfo->Get();
+				}
+			}
+			if (!Info)
+			{
+				return nullptr;
+			}
+
+			// Load on-demand if we need to
+			//
+			// Only try to do this once, and set bFailed to true if we failed, so that we can assume
+			// that nothing ever writes to this after we first try this on-demand load. With this assumption,
+			// we don't need to use read locks everywhere
+			{
+				FWriteScopeLock WriteLock(FileLocks[Info->LockIndex % GMaxParallelFileReads]);
+
+				if (!Info->bFailed && Info->FileBytes.Num() == 0)
+				{
+					if (!FPaths::FileExists(Filename))
+					{
+						UE_LOG(LogInterchangeOpenVDBImport, Error, TEXT("OpenVDB file could not be found: %s"), *Filename);
+						Info->bFailed = true;
+						return nullptr;
+					}
+
+					if (!FFileHelper::LoadFileToArray(Info->FileBytes, *Filename))
+					{
+						UE_LOG(LogInterchangeOpenVDBImport, Error, TEXT("OpenVDB file could not be loaded: %s"), *Filename);
+						Info->bFailed = true;
+						return nullptr;
+					}
+
+					const bool bCreateStrings = true;
+					if (!GetOpenVDBGridInfo(Info->FileBytes, bCreateStrings, &Info->GridInfo))
+					{
+						UE_LOG(LogInterchangeOpenVDBImport, Error, TEXT("Failed to read OpenVDB file: %s"), *Filename);
+						Info->bFailed = true;
+						return nullptr;
+					}
+				}
+			}
+
+			return Info;
+		}
+
+		const TMap<FString, TUniquePtr<FOpenVDBFileInfo>>& GetAllFileInfo() const
+		{
+			return FileNameToInfo;
+		}
+
+	private:
+		// We don't need this when translating files directly, as we can ensure FileNameToInfo only receives
+		// new entries on the game thread. If this translator is used indirectly however (i.e. through another
+		// translator like USDTranslator) we may be asked to return payload data for a file we never seen
+		// before and may need to open on-demand. Since payload retrieval can be multi-threaded, we may
+		// have multiple threads trying to add entries to FileNameToInfo at the same time, which is what
+		// this lock is for.
+		FRWLock FileNameToInfoLock;
+		TMap<FString, TUniquePtr<FOpenVDBFileInfo>> FileNameToInfo;
+
+		FRWLock* FileLocks = nullptr;
+		int32 NextFileLockIndex = 0;
+	};
+#else
+	class UInterchangeOpenVDBTranslatorImpl
+	{
+	};
+#endif	  // WITH_EDITOR
+}	 // namespace UE::InterchangeOpenVDBTranslator::Private
+
+UInterchangeOpenVDBTranslator::UInterchangeOpenVDBTranslator()
+	: Impl(MakeUnique<UE::InterchangeOpenVDBTranslator::Private::UInterchangeOpenVDBTranslatorImpl>())
+{
+}
+
+EInterchangeTranslatorType UInterchangeOpenVDBTranslator::GetTranslatorType() const
+{
+	return EInterchangeTranslatorType::Assets;
+}
+
+EInterchangeTranslatorAssetType UInterchangeOpenVDBTranslator::GetSupportedAssetTypes() const
+{
+	return EInterchangeTranslatorAssetType::Textures;
+}
+
+bool UInterchangeOpenVDBTranslator::CanImportSourceData(const UInterchangeSourceData* InSourceData) const
+{
+	// Ignore the cvar's effect from GetSupportedFormats() in case we're being used by the USD translator
+	// to import USD+OpenVDB files (it will add this additional context object)
+	UObject* ContextObject = InSourceData->GetContextObjectByTag(UE::Interchange::USD::USDContextTag);
+	if (ContextObject)
+	{
+		const bool bIncludeDot = false;
+		const FString Extension = FPaths::GetExtension(InSourceData->GetFilename(), bIncludeDot).ToLower();
+		if (Extension == TEXT("vdb"))
+		{
+			return true;
+		}
+	}
+
+	return Super::CanImportSourceData(InSourceData);
+}
+
+TArray<FString> UInterchangeOpenVDBTranslator::GetSupportedFormats() const
+{
+	TArray<FString> Extensions;
+#if WITH_EDITOR
+	if (GInterchangeEnableOpenVDBImport)
+	{
+		Extensions.Add(TEXT("vdb;OpenVDB files"));
+	}
+#endif	  // WITH_EDITOR
+	return Extensions;
+}
+
+bool UInterchangeOpenVDBTranslator::Translate(UInterchangeBaseNodeContainer& NodeContainer) const
+{
+#if WITH_EDITOR && OPENVDB_AVAILABLE
+	using namespace UE::InterchangeOpenVDBTranslator::Private;
+
+	// References:
+	// - SparseVolumeTextureFactory.cpp, LoadOpenVDBPreviewData()
+
+	UInterchangeOpenVDBTranslatorImpl* ImplPtr = Impl.Get();
+	if (!ImplPtr)
+	{
+		return false;
+	}
+
+	UInterchangeVolumeTranslatorSettings* Settings = Cast<UInterchangeVolumeTranslatorSettings>(GetSettings());
+	if (!Settings)
+	{
+		return false;
+	}
+
+	// Get all relevant VDB filenames to import
+	TArray<FString> Filenames;
+	{
+		const FString MainFilename = GetSourceData()->GetFilename();
+
+		if (Settings->bTranslateAdjacentNumberedFiles)
+		{
+			Filenames = FindOpenVDBSequenceFileNames(MainFilename);
+		}
+		else
+		{
+			Filenames = {MainFilename};
+		}
+	}
+
+	// Group up multiple VDB files in a single animation ID, so that the SVT pipeline makes one animated SVT factory node out of the group.
+	//
+	// We try using the ID specified in the settings first, because it's handy to be able to drive this via the USD translator as
+	// the AnimationID will make it's way into the volume / grid node UIDs, which is cumbersome to patch up afterwards
+	FString AnimationID = Settings->AnimationID;
+	if (AnimationID.IsEmpty() && Filenames.Num() > 1)
+	{
+		FSHA1 SHA1;
+		for (const FString& Filename : Filenames)
+		{
+			SHA1.UpdateWithString(*Filename, Filename.Len());
+		}
+		SHA1.Final();
+
+		FSHAHash FilePathHash;
+		SHA1.GetHash(&FilePathHash.Hash[0]);
+		AnimationID = FilePathHash.ToString();
+	}
+
+	// Create all the entries we'll need in one thread
+	ImplPtr->ClearFileInfoEntries();
+	ImplPtr->EnsureFileInfoEntries(Filenames);
+
+	// Collect all the frame info in parallel
+	const int32 NumBatches = Filenames.Num();
+	ParallelFor(
+		NumBatches,
+		[ImplPtr, &Filenames](int32 FilenameIndex)
+		{
+			const FString& Filename = Filenames[FilenameIndex];
+			ImplPtr->GetOrLoadFileInfo(Filename);
+		}
+	);
+
+	// Note that Filenames are sorted inside FindOpenVDBSequenceFileNames, so this is also our volume frame index into the animation,
+	// if we do have any (e.g. filenames [tornado_23.vdb, tornado_47.vdb, tornado_77.vdb] --> indices [0, 1, 2])
+	for (int32 FilenameIndex = 0; FilenameIndex < Filenames.Num(); ++FilenameIndex)
+	{
+		const FString& Filename = Filenames[FilenameIndex];
+		const FOpenVDBFileInfo* FileInfo = ImplPtr->GetOrLoadFileInfo(Filenames[FilenameIndex]);
+		if (!FileInfo)
+		{
+			continue;
+		}
+
+		const static FString VolumePrefix = TEXT("\\Volume\\");
+		const FString BaseFilename = FPaths::GetBaseFilename(Filename);
+
+		// Full filename helps prevent name collisions.
+		// AnimationID because the same volume may be used for different animations in the same import (e.g. via USD). In that case,
+		// we want to emit a factory node for each animation, but it is nice to retain the correspondence between translated node
+		// and factory node UIDs of just having the extra factory node prefix. That directly implies we need a separate volume node
+		// for each animation ID, and so the animation ID must be part of the node UID
+		const FString VolumeNodeUid = VolumePrefix + Filename + TEXT("\\") + AnimationID;
+
+		// Generate node for this file
+		UInterchangeVolumeNode* VolumeNode = nullptr;
+		{
+			const UInterchangeVolumeNode* Node = Cast<UInterchangeVolumeNode>(NodeContainer.GetNode(VolumeNodeUid));
+			if (Node)
+			{
+				continue;
+			}
+
+			VolumeNode = NewObject<UInterchangeVolumeNode>(&NodeContainer);
+			VolumeNode->InitializeNode(VolumeNodeUid, BaseFilename, EInterchangeNodeContainerType::TranslatedAsset);
+			VolumeNode->SetAssetName(VolumeNodeUid);
+			VolumeNode->SetCustomFileName(Filename);
+			if (!AnimationID.IsEmpty())
+			{
+				VolumeNode->SetCustomAnimationID(AnimationID);
+				VolumeNode->AddCustomFrameIndexInAnimation(FilenameIndex);
+			}
+
+			NodeContainer.AddNode(VolumeNode);
+		}
+
+		// Generate nodes for each volume grid
+		for (const FOpenVDBGridInfo& GridInfo : FileInfo->GridInfo)
+		{
+			const FString& GridName = GridInfo.Name;
+			const FString GridNodeUid = VolumeNodeUid + TEXT("_") + GridName;
+
+			const UInterchangeVolumeGridNode* Node = Cast<UInterchangeVolumeGridNode>(NodeContainer.GetNode(VolumeNodeUid));
+			if (Node)
+			{
+				continue;
+			}
+
+			UInterchangeVolumeGridNode* GridNode = NewObject<UInterchangeVolumeGridNode>(&NodeContainer);
+			GridNode->InitializeNode(GridNodeUid, GridName, EInterchangeNodeContainerType::TranslatedAsset);
+			GridNode->SetCustomElementType(GridTypeToInterchangeGridType(GridInfo.Type));
+			GridNode->SetCustomNumComponents(GridInfo.NumComponents);
+			GridNode->SetCustomGridTransform(GridInfo.Transform);
+			GridNode->SetCustomGridActiveAABBMin(GridInfo.VolumeActiveAABBMin);
+			GridNode->SetCustomGridActiveAABBMax(GridInfo.VolumeActiveAABBMax);
+			GridNode->SetCustomGridActiveDimensions(GridInfo.VolumeActiveDim);
+
+			VolumeNode->AddCustomGridDependency(GridNodeUid);
+
+			NodeContainer.AddNode(GridNode);
+			NodeContainer.SetNodeParentUid(GridNodeUid, VolumeNodeUid);
+		}
+	}
+
+	return true;
+#else
+	return false;
+#endif	  // WITH_EDITOR && OPENVDB_AVAILABLE
+}
+
+void UInterchangeOpenVDBTranslator::ReleaseSource()
+{
+#if WITH_EDITOR
+	UE::InterchangeOpenVDBTranslator::Private::UInterchangeOpenVDBTranslatorImpl* ImplPtr = Impl.Get();
+	if (!ImplPtr)
+	{
+		return;
+	}
+
+	ImplPtr->ClearFileInfoEntries();
+
+	if (TranslatorSettings)
+	{
+		TranslatorSettings->ClearFlags(RF_Standalone);
+		TranslatorSettings = nullptr;
+	}
+#endif	  // WITH_EDITOR
+}
+
+UInterchangeTranslatorSettings* UInterchangeOpenVDBTranslator::GetSettings() const
+{
+#if WITH_EDITOR
+	using namespace UE::InterchangeOpenVDBTranslator::Private;
+
+	UInterchangeOpenVDBTranslatorImpl* ImplPtr = Impl.Get();
+	if (!ImplPtr)
+	{
+		return nullptr;
+	}
+
+	if (!TranslatorSettings)
+	{
+		TranslatorSettings = DuplicateObject<UInterchangeVolumeTranslatorSettings>(
+			UInterchangeVolumeTranslatorSettings::StaticClass()->GetDefaultObject<UInterchangeVolumeTranslatorSettings>(),
+			GetTransientPackage()
+		);
+		TranslatorSettings->LoadSettings();
+		TranslatorSettings->ClearFlags(RF_ArchetypeObject);
+		TranslatorSettings->SetFlags(RF_Standalone);
+		TranslatorSettings->ClearInternalFlags(EInternalObjectFlags::Async);
+	}
+	return TranslatorSettings;
+#else
+	return nullptr;
+#endif	  // WITH_EDITOR
+}
+
+void UInterchangeOpenVDBTranslator::SetSettings(const UInterchangeTranslatorSettings* InterchangeTranslatorSettings)
+{
+#if WITH_EDITOR
+	using namespace UE::InterchangeOpenVDBTranslator::Private;
+
+	UInterchangeOpenVDBTranslatorImpl* ImplPtr = Impl.Get();
+	if (!ImplPtr)
+	{
+		return;
+	}
+
+	if (TranslatorSettings)
+	{
+		TranslatorSettings->ClearFlags(RF_Standalone);
+		TranslatorSettings->ClearInternalFlags(EInternalObjectFlags::Async);
+		TranslatorSettings = nullptr;
+	}
+	if (const UInterchangeVolumeTranslatorSettings* OpenVDBTranslatorSettings = Cast<UInterchangeVolumeTranslatorSettings>(
+			InterchangeTranslatorSettings
+		))
+	{
+		TranslatorSettings = DuplicateObject<UInterchangeVolumeTranslatorSettings>(OpenVDBTranslatorSettings, GetTransientPackage());
+		TranslatorSettings->ClearInternalFlags(EInternalObjectFlags::Async);
+		TranslatorSettings->SetFlags(RF_Standalone);
+	}
+#endif	  // WITH_EDITOR
+}
+
+TOptional<UE::Interchange::FVolumePayloadData> UInterchangeOpenVDBTranslator::GetVolumePayloadData(
+	const UE::Interchange::FVolumePayloadKey& PayloadKey
+) const
+{
+#if WITH_EDITOR && OPENVDB_AVAILABLE
+	using namespace UE::InterchangeOpenVDBTranslator::Private;
+
+	UInterchangeOpenVDBTranslatorImpl* ImplPtr = Impl.Get();
+	if (!ImplPtr)
+	{
+		return {};
+	}
+
+	// Note: We shouldn't need to lock FoundInfo's mutex to read from it after we got it:
+	// We're not going to write to it after the initial load anyway
+	const FOpenVDBFileInfo* FoundInfo = nullptr;
+	{
+		ImplPtr->EnsureFileInfoEntries({PayloadKey.FileName});
+
+		FoundInfo = ImplPtr->GetOrLoadFileInfo(PayloadKey.FileName);
+		if (!FoundInfo)
+		{
+			return {};
+		}
+	}
+
+	// Convert our custom struct to the one used by ConvertOpenVDBToSparseVolumeTexture
+	FOpenVDBImportOptions ImportOptions;
+	{
+		ImportOptions.bIsSequence = PayloadKey.AssignmentInfo.bIsSequence;
+
+		for (int32 TextureIndex = 0; TextureIndex < ImportOptions.Attributes.Num(); ++TextureIndex)
+		{
+			const UE::Interchange::Volume::FTextureInfo& InTextureInfo = PayloadKey.AssignmentInfo.Attributes[TextureIndex];
+			FOpenVDBSparseVolumeAttributesDesc& OutTextureInfo = ImportOptions.Attributes[TextureIndex];
+
+			static_assert((int)ESparseVolumeAttributesFormat::Unorm8 == (int)EInterchangeSparseVolumeTextureFormat::Unorm8);
+			static_assert((int)ESparseVolumeAttributesFormat::Float16 == (int)EInterchangeSparseVolumeTextureFormat::Float16);
+			static_assert((int)ESparseVolumeAttributesFormat::Float32 == (int)EInterchangeSparseVolumeTextureFormat::Float32);
+			OutTextureInfo.Format = static_cast<ESparseVolumeAttributesFormat>(InTextureInfo.Format);
+
+			for (int32 ChannelIndex = 0; ChannelIndex < OutTextureInfo.Mappings.Num(); ++ChannelIndex)
+			{
+				OutTextureInfo.Mappings[ChannelIndex].SourceComponentIndex = InTextureInfo.Mappings[ChannelIndex].SourceComponentIndex;
+				OutTextureInfo.Mappings[ChannelIndex].SourceGridIndex = InTextureInfo.Mappings[ChannelIndex].SourceGridIndex;
+			}
+		}
+	}
+
+	UE::Interchange::FVolumePayloadData Result;
+	bool bSucess = ConvertOpenVDBToSparseVolumeTexture(
+		const_cast<TArray64<uint8>&>(FoundInfo->FileBytes),	   // This function needs non-const access but it shouldn't ever write to anything
+		ImportOptions,
+		PayloadKey.VolumeBoundsMin,
+		Result.TextureData,
+		Result.Transform
+	);
+	if (bSucess)
+	{
+		return Result;
+	}
+#endif	  // WITH_EDITOR && OPENVDB_AVAILABLE
+
+	return {};
+}
+
+#undef LOCTEXT_NAMESPACE
diff --git a/Engine/Plugins/Interchange/Extensions/OpenVDB/Source/Import/Public/InterchangeOpenVDBImportModule.h b/Engine/Plugins/Interchange/Extensions/OpenVDB/Source/Import/Public/InterchangeOpenVDBImportModule.h
new file mode 100644
--- /dev/null
+++ b/Engine/Plugins/Interchange/Extensions/OpenVDB/Source/Import/Public/InterchangeOpenVDBImportModule.h
@@ -0,0 +1,11 @@
+// Copyright Epic Games, Inc. All Rights Reserved.
+#pragma once
+
+#include "Modules/ModuleInterface.h"
+
+class FInterchangeOpenVDBImportModule : public IModuleInterface
+{
+public:
+	virtual void StartupModule() override;
+	virtual void ShutdownModule() override;
+};
diff --git a/Engine/Plugins/Interchange/Extensions/OpenVDB/Source/Import/Public/InterchangeOpenVDBTranslator.h b/Engine/Plugins/Interchange/Extensions/OpenVDB/Source/Import/Public/InterchangeOpenVDBTranslator.h
new file mode 100644
--- /dev/null
+++ b/Engine/Plugins/Interchange/Extensions/OpenVDB/Source/Import/Public/InterchangeOpenVDBTranslator.h
@@ -0,0 +1,53 @@
+// Copyright Epic Games, Inc. All Rights Reserved.
+
+#pragma once
+
+#include "InterchangeTranslatorBase.h"
+#include "Volume/InterchangeVolumePayloadInterface.h"
+
+#include "UObject/ObjectMacros.h"
+
+#include "InterchangeOpenVDBTranslator.generated.h"
+
+namespace UE::Interchange
+{
+	struct FVolumePayloadData;
+	struct FVolumePayloadKey;
+}
+namespace UE::InterchangeOpenVDBTranslator::Private
+{
+	class UInterchangeOpenVDBTranslatorImpl;
+}
+class UInterchangeVolumeTranslatorSettings;
+
+UCLASS(BlueprintType)
+class UInterchangeOpenVDBTranslator
+	: public UInterchangeTranslatorBase
+	, public IInterchangeVolumePayloadInterface
+{
+	GENERATED_BODY()
+
+public:
+	UInterchangeOpenVDBTranslator();
+
+	/** Begin UInterchangeTranslatorBase interface */
+	virtual EInterchangeTranslatorType GetTranslatorType() const override;
+	virtual EInterchangeTranslatorAssetType GetSupportedAssetTypes() const override;
+	virtual bool CanImportSourceData(const UInterchangeSourceData* InSourceData) const override;
+	virtual TArray<FString> GetSupportedFormats() const override;
+	virtual bool Translate(UInterchangeBaseNodeContainer& BaseNodeContainer) const override;
+	virtual void ReleaseSource() override;
+	virtual UInterchangeTranslatorSettings* GetSettings() const override;
+	virtual void SetSettings(const UInterchangeTranslatorSettings* InterchangeTranslatorSettings) override;
+	/** End UInterchangeTranslatorBase interface */
+
+	/** Begin Interchange payload interfaces */
+	virtual TOptional<UE::Interchange::FVolumePayloadData> GetVolumePayloadData(const UE::Interchange::FVolumePayloadKey& PayloadKey) const override;
+	/** End Interchange payload interfaces */
+
+private:
+	TUniquePtr<UE::InterchangeOpenVDBTranslator::Private::UInterchangeOpenVDBTranslatorImpl> Impl;
+
+	UPROPERTY(Transient, DuplicateTransient)
+	mutable TObjectPtr<UInterchangeVolumeTranslatorSettings> TranslatorSettings = nullptr;
+};
diff --git a/Engine/Plugins/Interchange/Runtime/Source/Import/InterchangeImport.Build.cs b/Engine/Plugins/Interchange/Runtime/Source/Import/InterchangeImport.Build.cs
--- a/Engine/Plugins/Interchange/Runtime/Source/Import/InterchangeImport.Build.cs
+++ b/Engine/Plugins/Interchange/Runtime/Source/Import/InterchangeImport.Build.cs
@@ -63,7 +63,6 @@ namespace UnrealBuildTool.Rules
 						"InterchangeFbxParser",
 						"MaterialEditor",
 						"SkeletalMeshUtilitiesCommon",
-						"SparseVolumeTexture",
 						"UnrealEd",
 						"VariantManager",
 					}
diff --git a/Engine/Plugins/Interchange/Runtime/Source/Import/Private/InterchangeImportModule.cpp b/Engine/Plugins/Interchange/Runtime/Source/Import/Private/InterchangeImportModule.cpp
--- a/Engine/Plugins/Interchange/Runtime/Source/Import/Private/InterchangeImportModule.cpp
+++ b/Engine/Plugins/Interchange/Runtime/Source/Import/Private/InterchangeImportModule.cpp
@@ -21,7 +21,6 @@
 #include "Mesh/InterchangeStaticMeshFactory.h"
 #include "Misc/CoreDelegates.h"
 #include "Modules/ModuleManager.h"
-#include "OpenVDB/InterchangeOpenVDBTranslator.h"
 #include "Scene/InterchangeActorFactory.h"
 #include "Scene/InterchangeCameraActorFactory.h"
 #include "Scene/InterchangeDecalActorFactory.h"
@@ -103,13 +102,6 @@ void FInterchangeImportModule::StartupModule()
 		InterchangeManager.RegisterTranslator(UInterchangePSDTranslator::StaticClass());
 		InterchangeManager.RegisterTranslator(UInterchangeIESTranslator::StaticClass());
 
-		//Volumes
-#if WITH_EDITOR && OPENVDB_AVAILABLE
-		// Editor and OpenVDB-only because we use GetOpenVDBGridInfo and ConvertOpenVDBToSparseVolumeTexture
-		// which use OpenVDB and are in an Editor-only module
-		InterchangeManager.RegisterTranslator(UInterchangeOpenVDBTranslator::StaticClass());
-#endif
-
 		//Register the factories
 		InterchangeManager.RegisterFactory(UInterchangeTextureFactory::StaticClass());
 		InterchangeManager.RegisterFactory(UInterchangeMaterialFactory::StaticClass());
@@ -139,7 +131,7 @@ void FInterchangeImportModule::StartupModule()
 		InterchangeManager.RegisterFactory(UInterchangeSparseVolumeTextureFactory::StaticClass());
 #endif
 	};
-	
+
 	if (GEngine)
 	{
 		RegisterItems();
diff --git a/Engine/Plugins/Interchange/Runtime/Source/Import/Private/OpenVDB/InterchangeOpenVDBTranslator.cpp b/Engine/Plugins/Interchange/Runtime/Source/Import/Private/OpenVDB/InterchangeOpenVDBTranslator.cpp
deleted file mode 100644
diff --git a/Engine/Plugins/Interchange/Runtime/Source/Import/Private/Usd/InterchangeUsdTranslator.cpp b/Engine/Plugins/Interchange/Runtime/Source/Import/Private/Usd/InterchangeUsdTranslator.cpp
--- a/Engine/Plugins/Interchange/Runtime/Source/Import/Private/Usd/InterchangeUsdTranslator.cpp
+++ b/Engine/Plugins/Interchange/Runtime/Source/Import/Private/Usd/InterchangeUsdTranslator.cpp
@@ -49,10 +49,10 @@
 #include "InterchangeVolumeNode.h"
 #include "Mesh/InterchangeMeshPayload.h"
 #include "Nodes/InterchangeUserDefinedAttribute.h"
-#include "OpenVDB/InterchangeOpenVDBTranslator.h"
 #include "Usd/InterchangeUsdContext.h"
 #include "Usd/InterchangeUsdDefinitions.h"
 #include "Volume/InterchangeVolumeDefinitions.h"
+#include "Volume/InterchangeVolumeTranslatorSettings.h"
 
 #include "MaterialX/MaterialXUtils/MaterialXBase.h"
 #include "MaterialX/MaterialXUtils/MaterialXManager.h"
@@ -971,7 +971,7 @@ namespace UE::InterchangeUsdTranslator::Private
 				}
 			}
 		);
-#endif // WITH_EDITOR
+#endif	  // WITH_EDITOR
 
 		// Now we need to check if we have to set the slot of the mesh nodes here
 		if (TArray<FMaterialSlotMesh>* SlotMeshes = PrimPathToSlotMeshNodes.Find(MaterialUid))
@@ -998,7 +998,7 @@ namespace UE::InterchangeUsdTranslator::Private
 							);
 						}
 					}
-#endif // WITH_EDITOR
+#endif	  // WITH_EDITOR
 					if (!MaterialSlotMesh.MeshNode->GetSlotMaterialDependencyUid(MaterialSlotMesh.MaterialSlotName, *NewMaterialUID))
 					{
 						MaterialSlotMesh.MeshNode->SetSlotMaterialDependencyUid(MaterialSlotMesh.MaterialSlotName, *NewMaterialUID);
@@ -2247,7 +2247,7 @@ namespace UE::InterchangeUsdTranslator::Private
 		// this means we need a separate volume node *per animation*, so we'll also use "AnimationIDs" to differentiate them
 
 		TStrongObjectPtr<UInterchangeBaseNodeContainer> TempContainer = nullptr;
-		TStrongObjectPtr<UInterchangeOpenVDBTranslator> Translator = nullptr;
+		TStrongObjectPtr<UInterchangeTranslatorBase> Translator = nullptr;
 
 		// This is collected by path hash here because for animated SVTs we want to still have a single UsdUtils::FVolumePrimInfo for each
 		// group of animated volume frames, since they will become a separate SVT. If we just collected them by filepath we could run into
@@ -2325,24 +2325,31 @@ namespace UE::InterchangeUsdTranslator::Private
 
 					if (!Translator || !Translator->CanImportSourceData(SourceData))
 					{
-						// Check and create an OpenVDB translator directly for the file. Ideally we'd switch the cvar
-						// ourselves and just use UInterchangeManager::GetTranslatorForSourceData(), however we can only toggle
-						// cvars from the game thread, and translation can happen on any thread it seems
-						const bool bIncludeDot = false;
-						const FString Extension = FPaths::GetExtension(FilePath, bIncludeDot);
-						if (Extension.ToLower() == TEXT("vdb"))
-						{
-							Translator.Reset(NewObject<UInterchangeOpenVDBTranslator>(GetTransientPackage(), NAME_None));
-							Translator->SourceData = SourceData;
-						}
+						// Pass a USD context object to the translator, which is a signal that lets the OpenVDB translator
+						// be considered, even if its cvar is off (c.f. UInterchangeOpenVDBTranslator::CanImportSourceData)
+						UInterchangeUsdContext* Context = NewObject<UInterchangeUsdContext>();
+						SourceData->SetContextObjectByTag(UE::Interchange::USD::USDContextTag, Context);
+
+						UInterchangeManager& InterchangeManager = UInterchangeManager::GetInterchangeManager();
+						Translator.Reset(InterchangeManager.GetTranslatorForSourceData(SourceData));
 					}
 					if (!Translator)
 					{
-						UE_LOG(LogUsd, Warning, TEXT("Failed to retrieve a translator for filepath '%s'"), *FilePath);
+						const bool bIncludeDot = false;
+						const FString Extension = FPaths::GetExtension(FilePath, bIncludeDot).ToLower();
+						const bool bIsOpenVDB = Extension == TEXT("vdb");
+
+						UE_LOG(
+							LogUsd,
+							Error,
+							TEXT("Failed to find a compatible translator for file '%s'.%s"),
+							*FilePath,
+							bIsOpenVDB ? TEXT(" Is the 'Interchange OpenVDB Import' plugin enabled?") : TEXT("")
+						);
 						continue;
 					}
 
-					if (UInterchangeOpenVDBTranslatorSettings* Settings = Cast<UInterchangeOpenVDBTranslatorSettings>(Translator->GetSettings()))
+					if (UInterchangeVolumeTranslatorSettings* Settings = Cast<UInterchangeVolumeTranslatorSettings>(Translator->GetSettings()))
 					{
 						// We never want to discover new .vdb files via the OpenVDB translator for animations.
 						// If we should have animations via USD they will be described on the USD file explicictly
diff --git a/Engine/Plugins/Interchange/Runtime/Source/Import/Public/OpenVDB/InterchangeOpenVDBTranslator.h b/Engine/Plugins/Interchange/Runtime/Source/Import/Public/OpenVDB/InterchangeOpenVDBTranslator.h
deleted file mode 100644
diff --git a/Engine/Plugins/Interchange/Runtime/Source/Import/Public/Volume/InterchangeVolumeTranslatorSettings.h b/Engine/Plugins/Interchange/Runtime/Source/Import/Public/Volume/InterchangeVolumeTranslatorSettings.h
new file mode 100644
--- /dev/null
+++ b/Engine/Plugins/Interchange/Runtime/Source/Import/Public/Volume/InterchangeVolumeTranslatorSettings.h
@@ -0,0 +1,32 @@
+// Copyright Epic Games, Inc. All Rights Reserved.
+
+#pragma once
+
+#include "InterchangeTranslatorBase.h"
+#include "UObject/ObjectMacros.h"
+
+#include "InterchangeVolumeTranslatorSettings.generated.h"
+
+UCLASS(BlueprintType, editinlinenew, MinimalAPI)
+class UInterchangeVolumeTranslatorSettings : public UInterchangeTranslatorSettings
+{
+	GENERATED_BODY()
+
+public:
+	/**
+	 * Whether we should search the same folder of the translated file for other .vdb files whose file names differ from it
+	 * only by a numbered suffix (e.g. "Tornado_0001.vdb, Tornado_0002.vdb, etc.), and translate those too.
+	 *
+	 * This is useful if you intend on importing an animated SparseVolumeTexture.
+	 */
+	UPROPERTY(EditAnywhere, Category = "OpenVDB Translator")
+	bool bTranslateAdjacentNumberedFiles = true;
+
+	/**
+	 * Identifier set on the generated volume nodes (via SetCustomAnimationID) if they are part of an animation, to identify said animation.
+	 * Volume nodes that share the same AnimationID belong to the same volume animation.
+	 * This can be left empty to have an ID automatically generated, or can be set to something to force it to that value.
+	 */
+	UPROPERTY()
+	FString AnimationID;
+};
