Interchange OpenVDB: Register pipeline stacks and show import dialog override on C++ when the module is loaded.

diff --git a/Engine/Plugins/Interchange/Extensions/OpenVDB/Source/Import/Private/InterchangeOpenVDBImportModule.cpp b/Engine/Plugins/Interchange/Extensions/OpenVDB/Source/Import/Private/InterchangeOpenVDBImportModule.cpp
--- a/Engine/Plugins/Interchange/Extensions/OpenVDB/Source/Import/Private/InterchangeOpenVDBImportModule.cpp
+++ b/Engine/Plugins/Interchange/Extensions/OpenVDB/Source/Import/Private/InterchangeOpenVDBImportModule.cpp
@@ -7,6 +7,7 @@
 
 #include "Engine/Engine.h"
 #include "InterchangeManager.h"
+#include "InterchangeProjectSettings.h"
 #include "Misc/CoreDelegates.h"
 #include "Modules/ModuleManager.h"
 
@@ -18,28 +19,62 @@ void FInterchangeOpenVDBImportModule::StartupModule()
 
 	auto RegisterItems = [this]()
 	{
-		UInterchangeManager& InterchangeManager = UInterchangeManager::GetInterchangeManager();
-
-		// Volumes
+// Editor and OpenVDB-only because we use GetOpenVDBGridInfo and ConvertOpenVDBToSparseVolumeTexture
+// which use OpenVDB and are in an Editor-only module
 #if WITH_EDITOR && OPENVDB_AVAILABLE
-		// Editor and OpenVDB-only because we use GetOpenVDBGridInfo and ConvertOpenVDBToSparseVolumeTexture
-		// which use OpenVDB and are in an Editor-only module
+		UInterchangeManager& InterchangeManager = UInterchangeManager::GetInterchangeManager();
 		InterchangeManager.RegisterTranslator(UInterchangeOpenVDBTranslator::StaticClass());
+
+		TSoftClassPtr<UInterchangeTranslatorBase> TranslatorClassPath = TSoftClassPtr<UInterchangeTranslatorBase>(
+			UInterchangeOpenVDBTranslator::StaticClass()
+		);
+
+		FInterchangeTranslatorPipelines TranslatorPipelines;
+		TranslatorPipelines.Translator = TranslatorClassPath;
+		TranslatorPipelines.Pipelines.Add(
+			FSoftObjectPath{TEXT("/Interchange/Pipelines/DefaultSparseVolumeTexturePipeline.DefaultSparseVolumeTexturePipeline")}
+		);
+
+		// Don't go through FInterchangeProjectSettingsUtils::GetMutableDefaultImportSettings as we'll need a member of the actual
+		// FInterchangeContentImportSettings struct anyway, and this is likely safer than casting the struct pointer
+		UInterchangeProjectSettings* ProjectSettings = GetMutableDefault<UInterchangeProjectSettings>();
+		if (!ProjectSettings)
+		{
+			return;
+		}
+
+		// Scene import pipeline stacks
+		FInterchangeImportSettings& SceneImportSettings = ProjectSettings->SceneImportSettings;
+		FInterchangePipelineStack& PipelineStack = SceneImportSettings.PipelineStacks[TEXT("Scene")];
+		PipelineStack.PerTranslatorPipelines.Add(TranslatorPipelines);
+
+		// Asset import pipeline stacks
+		FInterchangeContentImportSettings& AssetImportSettings = ProjectSettings->ContentImportSettings;
+		FInterchangePipelineStack& AssetStack = AssetImportSettings.PipelineStacks[TEXT("Assets")];
+		AssetStack.PerTranslatorPipelines.Add(TranslatorPipelines);
+		FInterchangePipelineStack& TextureStack = AssetImportSettings.PipelineStacks[TEXT("Textures")];
+		TextureStack.PerTranslatorPipelines.Add(TranslatorPipelines);
+
+		// Asset import show dialog override (otherwise it doesn't show the import options dialog as OpenVDBs are Texture type assets)
+		{
+			FInterchangeDialogOverride& DialogOverrides = AssetImportSettings.ShowImportDialogOverride.FindOrAdd(
+				EInterchangeTranslatorAssetType::Textures
+			);
+
+			FInterchangePerTranslatorDialogOverride ImportDialogOverride;
+			ImportDialogOverride.Translator = TranslatorClassPath;
+			ImportDialogOverride.bShowImportDialog = true;
+			DialogOverrides.PerTranslatorImportDialogOverride.Add(ImportDialogOverride);
+		}
 #endif
 	};
 
-	if (GEngine)
-	{
-		RegisterItems();
-	}
-	else
-	{
-		FCoreDelegates::OnPostEngineInit.AddLambda(RegisterItems);
-	}
+	FCoreDelegates::OnPostEngineInit.AddLambda(RegisterItems);
 }
 
 void FInterchangeOpenVDBImportModule::ShutdownModule()
 {
+	FCoreDelegates::OnPostEngineInit.Remove(PostEngineInitHandle);
 }
 
 IMPLEMENT_MODULE(FInterchangeOpenVDBImportModule, InterchangeOpenVDBImport)
diff --git a/Engine/Plugins/Interchange/Extensions/OpenVDB/Source/Import/Public/InterchangeOpenVDBImportModule.h b/Engine/Plugins/Interchange/Extensions/OpenVDB/Source/Import/Public/InterchangeOpenVDBImportModule.h
--- a/Engine/Plugins/Interchange/Extensions/OpenVDB/Source/Import/Public/InterchangeOpenVDBImportModule.h
+++ b/Engine/Plugins/Interchange/Extensions/OpenVDB/Source/Import/Public/InterchangeOpenVDBImportModule.h
@@ -2,10 +2,14 @@
 #pragma once
 
 #include "Modules/ModuleInterface.h"
+#include "Delegates/IDelegateInstance.h"
 
 class FInterchangeOpenVDBImportModule : public IModuleInterface
 {
 public:
 	virtual void StartupModule() override;
 	virtual void ShutdownModule() override;
+
+private:
+	FDelegateHandle PostEngineInitHandle;
 };
