OpenVDB: Hide and disable the Interchange OpenVDB plugin for now, as it's only meant to be used by the Interchange OpenUSD plugin.

diff --git a/Engine/Plugins/Interchange/Extensions/OpenVDB/InterchangeOpenVDB.uplugin b/Engine/Plugins/Interchange/Extensions/OpenVDB/InterchangeOpenVDB.uplugin
--- a/Engine/Plugins/Interchange/Extensions/OpenVDB/InterchangeOpenVDB.uplugin
+++ b/Engine/Plugins/Interchange/Extensions/OpenVDB/InterchangeOpenVDB.uplugin
@@ -14,6 +14,7 @@
 	"CanContainContent": false,
 	"IsBetaVersion": false,
 	"Installed": false,
+	"Hidden": true,
 	"IsExperimentalVersion": true,
 	"SupportedTargetPlatforms": [
 		"Win64",
diff --git a/Engine/Plugins/Interchange/Extensions/OpenVDB/Source/Import/Private/InterchangeOpenVDBTranslator.cpp b/Engine/Plugins/Interchange/Extensions/OpenVDB/Source/Import/Private/InterchangeOpenVDBTranslator.cpp
--- a/Engine/Plugins/Interchange/Extensions/OpenVDB/Source/Import/Private/InterchangeOpenVDBTranslator.cpp
+++ b/Engine/Plugins/Interchange/Extensions/OpenVDB/Source/Import/Private/InterchangeOpenVDBTranslator.cpp
@@ -21,13 +21,6 @@
 
 #define LOCTEXT_NAMESPACE "InterchangeOpenVDBTranslator"
 
-static bool GInterchangeEnableOpenVDBImport = false;
-static FAutoConsoleVariableRef CVarInterchangeEnableOpenVDBImport(
-	TEXT("Interchange.FeatureFlags.Import.OpenVDB"),
-	GInterchangeEnableOpenVDBImport,
-	TEXT("Whether OpenVDB support is enabled via Interchange")
-);
-
 static int32 GMaxParallelFileReads = 16;
 static FAutoConsoleVariableRef CVarNumPerPrimLocks(
 	TEXT("Interchange.FeatureFlags.Import.OpenVDB.NumFileReadLocks"),
@@ -244,12 +237,11 @@ bool UInterchangeOpenVDBTranslator::CanImportSourceData(const UInterchangeSource
 TArray<FString> UInterchangeOpenVDBTranslator::GetSupportedFormats() const
 {
 	TArray<FString> Extensions;
-#if WITH_EDITOR
-	if (GInterchangeEnableOpenVDBImport)
-	{
-		Extensions.Add(TEXT("vdb;OpenVDB files"));
-	}
-#endif	  // WITH_EDITOR
+
+	// We don't advertise support for .vdb files here. The UInterchangeOpenVDBTranslator is exclusively used
+	// by the UInterchangeUSDTranslator for now, and shouldn't import .vdb files directly
+	// Extensions.Add(TEXT("vdb;OpenVDB files"));
+
 	return Extensions;
 }
 
